Adding the HMI template for project creation.
#android

#virtualized

diff --git a/Templates/TP_HMI_Blank/Config/DefaultDeviceProfiles.ini b/Templates/TP_HMI_Blank/Config/DefaultDeviceProfiles.ini
new file mode 100644
--- /dev/null
+++ b/Templates/TP_HMI_Blank/Config/DefaultDeviceProfiles.ini
@@ -0,0 +1,47 @@
+; ------------------------------------------------------------------------------
+; Android: Rules for matching devices
+; ------------------------------------------------------------------------------
+
+[Android DeviceProfile]
+DeviceType=Android
+BaseProfileName=
++CVars=r.MobileContentScaleFactor=0
++CVars=r.BloomQuality=2
++CVars=r.DepthOfFieldQuality=0
++CVars=r.LightShaftQuality=0
++CVars=r.RefractionQuality=0
++CVars=r.ShadowQuality=0
++CVars=slate.AbsoluteIndices=1
++CVars=r.Vulkan.DelayAcquireBackBuffer=0
++CVars=r.Vulkan.RobustBufferAccess=1
++CVars=r.Vulkan.DescriptorSetLayoutMode=2
++CVars=r.Android.DisableVulkanSupport=1
++CVars=r.ProgramBinaryCache.RestartAndroidAfterPrecompile=0
++CVars=r.DefaultFeature.AntiAliasing=2
++CVars=r.MobileMSAA=1
+
+;Binary Shader cache
++CVars=r.OpenGL.EnableProgramLRUCache=0
++CVars=r.OpenGL.ProgramLRUCount=10000
++CVars=r.OpenGL.ProgramLRUBinarySize=1000000000
++CVars=r.OpenGL.EvictOnBSSDestruct=1
++CVars=r.OpenGL.EvictOnBSSDestruct.Latency=200
++CVars=r.OpenGL.ProgramLRUKeepBinaryResident=1
++CVars=r.OpenGL.StoreCompressedProgramBinaries=0
+
+; ------------------------------------------------------------------------------
+; QNX: Rules for matching devices
+; ------------------------------------------------------------------------------
+
+[QNX DeviceProfile]
+DeviceType=QNX
+BaseProfileName=QNX
++CVars=r.MobileContentScaleFactor=0
++CVars=r.BloomQuality=2
++CVars=r.DepthOfFieldQuality=0
++CVars=r.LightShaftQuality=0
++CVars=r.RefractionQuality=0
++CVars=r.ShadowQuality=0
++CVars=r.DefaultFeature.AntiAliasing=2
++CVars=r.MobileMSAA=1
+
diff --git a/Templates/TP_HMI_Blank/Config/DefaultEngine.ini b/Templates/TP_HMI_Blank/Config/DefaultEngine.ini
new file mode 100644
--- /dev/null
+++ b/Templates/TP_HMI_Blank/Config/DefaultEngine.ini
@@ -0,0 +1,214 @@
+[/Script/Engine.RendererSettings]
+r.screenpercentage=100
+r.Mobile.DisableVertexFog=True
+r.Shadow.CSM.MaxMobileCascades=1
+r.Mobile.UseLegacyShadingModel=False
+r.Mobile.AllowDitheredLODTransition=False
+r.Mobile.AllowSoftwareOcclusion=False
+r.Mobile.VirtualTextures=False
+r.DiscardUnusedQuality=False
+r.AllowOcclusionQueries=False
+r.MinScreenRadiusForLights=0.030000
+r.MinScreenRadiusForDepthPrepass=0.030000
+r.MinScreenRadiusForCSMDepth=0.010000
+r.PrecomputedVisibilityWarning=False
+r.TextureStreaming=True
+Compat.UseDXT5NormalMaps=False
+r.VirtualTextures=False
+r.VirtualTexturedLightmaps=False
+r.VT.TileSize=128
+r.VT.TileBorderSize=4
+r.vt.FeedbackFactor=16
+r.VT.EnableCompressZlib=True
+r.VT.EnableCompressCrunch=False
+r.ClearCoatNormal=False
+r.AnisotropicBRDF=False
+r.ReflectionEnvironmentLightmapMixBasedOnRoughness=False
+r.ForwardShading=False
+r.VertexFoggingForOpaque=True
+r.AllowStaticLighting=True
+r.NormalMapsForStaticLighting=False
+r.GenerateMeshDistanceFields=False
+r.DistanceFieldBuild.EightBit=False
+r.GenerateLandscapeGIData=False
+r.DistanceFieldBuild.Compress=False
+r.TessellationAdaptivePixelsPerTriangle=48.000000
+r.SeparateTranslucency=False
+r.TranslucentSortPolicy=0
+TranslucentSortAxis=(X=0.000000,Y=-1.000000,Z=0.000000)
+r.CustomDepth=3
+r.CustomDepthTemporalAAJitter=True
+r.PostProcessing.PropagateAlpha=0
+r.DefaultFeature.Bloom=True
+r.DefaultFeature.AmbientOcclusion=False
+r.DefaultFeature.AmbientOcclusionStaticFraction=False
+r.DefaultFeature.AutoExposure=False
+r.DefaultFeature.AutoExposure.Method=0
+r.DefaultFeature.AutoExposure.Bias=1.000000
+r.DefaultFeature.AutoExposure.ExtendDefaultLuminanceRange=False
+r.UsePreExposure=True
+r.EyeAdaptation.EditorOnly=False
+r.EyeAdaptationQuality=0
+r.DefaultFeature.MotionBlur=False
+r.DefaultFeature.LensFlare=False
+r.TemporalAA.Upsampling=False
+r.SSGI.Enable=False
+r.DefaultFeature.AntiAliasing=2
+r.DefaultFeature.LightUnits=1
+r.DefaultBackBufferPixelFormat=4
+r.Shadow.UnbuiltPreviewInGame=True
+r.StencilForLODDither=False
+r.EarlyZPass=3
+r.EarlyZPassOnlyMaterialMasking=False
+r.DBuffer=False
+r.ClearSceneMethod=1
+r.BasePassOutputsVelocity=False
+r.VertexDeformationOutputsVelocity=False
+r.SelectiveBasePassOutputs=False
+bDefaultParticleCutouts=False
+fx.GPUSimulationTextureSizeX=1024
+fx.GPUSimulationTextureSizeY=1024
+r.AllowGlobalClipPlane=False
+r.GBufferFormat=1
+r.MorphTarget.Mode=True
+r.GPUCrashDebugging=False
+vr.InstancedStereo=False
+r.MobileHDR=True
+vr.MobileMultiView=False
+r.Mobile.UseHWsRGBEncoding=False
+vr.RoundRobinOcclusion=False
+vr.ODSCapture=False
+r.MeshStreaming=False
+r.WireframeCullThreshold=5.000000
+r.RayTracing=False
+r.RayTracing.UseTextureLod=False
+r.SupportStationarySkylight=True
+r.SupportLowQualityLightmaps=True
+r.SupportPointLightWholeSceneShadows=False
+r.SupportAtmosphericFog=False
+r.SupportSkyAtmosphere=False
+r.SupportSkyAtmosphereAffectsHeightFog=False
+r.SkinCache.CompileShaders=False
+r.SkinCache.DefaultBehavior=1
+r.SkinCache.SceneMemoryLimitInMB=128.000000
+r.SkinCache.ForceRecomputeTangents=False
+r.Mobile.EnableStaticAndCSMShadowReceivers=False
+r.Mobile.EnableMovableLightCSMShaderCulling=False
+r.Mobile.AllowDistanceFieldShadows=False
+r.Mobile.AllowMovableDirectionalLights=False
+r.MobileNumDynamicPointLights=1
+r.MobileDynamicPointLightsUseStaticBranch=True
+r.Mobile.EnableMovableSpotlights=False
+r.GPUSkin.Support16BitBoneIndex=False
+r.GPUSkin.Limit2BoneInfluences=False
+r.SupportDepthOnlyIndexBuffers=True
+r.SupportReversedIndexBuffers=True
+r.SupportMaterialLayers=False
+r.LightPropagationVolume=False
+r.Mobile.PropagateAlpha=1
+r.ReflectionCaptureResolution=512
+r.MobileContentScaleFactor=0
+r.MobileMSAA=1
+r.OpenGL.MaxShaderLibProcessingTime=1000
+
+[DevOptions.Shaders]
+NeedsShaderStableKeys=true
+
+[PBCUseWithoutShaderPipelineCache]
+;Needed by the automated PSO caching system
+PBC_GUID=(X=1, Y=2, Z=3, W=4);
+
+[ConsoleVariables]
+;Needed by the automated PSO caching system
+r.ProgramBinaryCache.UseWithoutShaderPipelineCache=1
+;HMI Optimizations
+s.nosound=1
+s.disableai=1
+s.disablestartupmovie=1
+
+[/Script/Engine.GarbageCollectionSettings]
+gc.MaxObjectsInGame=35000
+
+[/Script/Slate.SlateSettings]
+bExplicitCanvasChildZOrder=True
+
+[/Script/HardwareTargeting.HardwareTargetingSettings]
+TargetedHardwareClass=Mobile
+AppliedTargetedHardwareClass=Mobile
+DefaultGraphicsPerformance=Maximum
+AppliedDefaultGraphicsPerformance=Maximum
+
+[/Script/EngineSettings.GameMapsSettings]
+EditorStartupMap=/Game/Default.Default
+LocalMapOptions=
+TransitionMap=None
+bUseSplitscreen=False
+TwoPlayerSplitscreenLayout=Horizontal
+ThreePlayerSplitscreenLayout=FavorTop
+GameInstanceClass=/Script/Engine.GameInstance
+ServerDefaultMap=/Engine/Maps/Entry.Entry
+GlobalDefaultGameMode=/Script/Engine.GameModeBase
+GlobalDefaultServerGameMode=None
+
+[/Script/AndroidRuntimeSettings.AndroidRuntimeSettings]
+MinSDKVersion=29
+bSupportsVulkan=True
+bDetectVulkanByDefault=True
+bSupportsVulkanSM5=True
+TargetSDKVersion=31
+bAllowIMU=False
+bShouldHideUI=True
+bVerifyOBBOnStartUp=False
+bBuildForArmV7=False
+bBuildForX8664=False
+bBuildForArm64=True
+bBuildForES31=True
+bPackageDataInsideApk=True
+bAllowLargeOBBFiles=True
+PackageName=com.epicgames.[PROJECT]
+
+[/Script/LinuxTargetPlatform.LinuxTargetSettings]
+SpatializationPlugin=
+ReverbPlugin=
+OcclusionPlugin=
+SoundCueCookQualityIndex=-1
+-TargetedRHIs=SF_VULKAN_SM5
+-TargetedRHIs=GLSL_150_ES31
++TargetedRHIs=SF_VULKAN_ES31
++TargetedRHIs=GLSL_ES3_1_ANDROID
+
+[/Script/WindowsTargetPlatform.WindowsTargetSettings]
+Compiler=Default
+-TargetedRHIs=PCD3D_SM5
+-TargetedRHIs=PCD3D_ES31
+-TargetedRHIs=GLSL_150_ES31
++TargetedRHIs=PCD3D_SM5
++TargetedRHIs=PCD3D_ES31
++TargetedRHIs=GLSL_150_ES31
+DefaultGraphicsRHI=DefaultGraphicsRHI_Default
+MinimumOSVersion=MSOS_Vista
+bTarget32Bit=False
+AudioSampleRate=48000
+AudioCallbackBufferFrameSize=1024
+AudioNumBuffersToEnqueue=1
+AudioMaxChannels=0
+AudioNumSourceWorkers=4
+SpatializationPlugin=
+ReverbPlugin=
+OcclusionPlugin=
+CompressionOverrides=(bOverrideCompressionTimes=False,DurationThreshold=5.000000,MaxNumRandomBranches=0,SoundCueQualityIndex=0)
+bUseAudioStreamCaching=False
+CacheSizeKB=0
+bResampleForDevice=False
+MaxSampleRate=0.000000
+HighSampleRate=0.000000
+MedSampleRate=0.000000
+LowSampleRate=0.000000
+MinSampleRate=0.000000
+CompressionQualityModifier=0.000000
+AutoStreamingThreshold=0.000000
+SoundCueCookQualityIndex=-1
+
+[/Script/Engine.Engine]
++ActiveGameNameRedirects=(OldGameName="TP_HMI_Blank",NewGameName="/Script/TP_HMI_Blank")
++ActiveGameNameRedirects=(OldGameName="/Script/TP_HMI_Blank",NewGameName="/Script/TP_HMI_Blank")
diff --git a/Templates/TP_HMI_Blank/Config/DefaultGame.ini b/Templates/TP_HMI_Blank/Config/DefaultGame.ini
new file mode 100644
--- /dev/null
+++ b/Templates/TP_HMI_Blank/Config/DefaultGame.ini
@@ -0,0 +1,19 @@
+[/Script/EngineSettings.GeneralProjectSettings]
+ProjectID=(A=1823396781,B=1298598681,C=1743498152,D=-2048051703)
+ProjectName=TP_HMI_Blank
+
+[/Script/UnrealEd.ProjectPackagingSettings]
+bCompressed=True
+Build=IfProjectHasCode
+BuildConfiguration=PPBC_Development
+bShareMaterialShaderCode=True
+bSharedMaterialNativeLibraries=True
+bSkipEditorContent=True
+bUseIoStore=True
+IncludePrerequisites=False
+bSkipMovies=True
+PakFileCompressionFormats=Oodle
+PakFileAdditionalCompressionOptions=-oodlemethod=Selkie -oodlelevel=Optimial1
+
+[StartupActions]
+bAddPacks=False
\ No newline at end of file
diff --git a/Templates/TP_HMI_Blank/Config/DefaultInput.ini b/Templates/TP_HMI_Blank/Config/DefaultInput.ini
new file mode 100644
--- /dev/null
+++ b/Templates/TP_HMI_Blank/Config/DefaultInput.ini
@@ -0,0 +1,6 @@
+
+
+[/Script/Engine.InputSettings]
+bUseMouseForTouch=True
+DefaultTouchInterface=None
+
diff --git a/Templates/TP_HMI_Blank/Config/DefaultPakFileRules.ini b/Templates/TP_HMI_Blank/Config/DefaultPakFileRules.ini
new file mode 100644
--- /dev/null
+++ b/Templates/TP_HMI_Blank/Config/DefaultPakFileRules.ini
@@ -0,0 +1,24 @@
+; These rules are applied in order, the first rule that applies per file is taken and no others are evaluated
+
+[ExcludeEngineContent]
+bExcludeFromPaks=true
+bOverrideChunkManifest=true
++Files=".../Engine/Content/EngineFonts/Faces/DroidSansFallback.ufont"
++Files=".../Engine/Content/Slate/Fonts/DroidSansFallback.ttf"
++Files=".../Engine/Content/Slate/Fonts/DroidSansMono.ttf"
++Files=".../Engine/Content/Slate/Fonts/NotoNaskhArabicUI-Regular.ttf"
++Files=".../Engine/Content/Slate/Fonts/Roboto-Black.ttf"
++Files=".../Engine/Content/Slate/Fonts/Roboto-BlackItalic.ttf"
++Files=".../Engine/Content/Slate/Fonts/Roboto-Bold.ttf"
++Files=".../Engine/Content/Slate/Fonts/Roboto-BoldCondensed.ttf"
++Files=".../Engine/Content/Slate/Fonts/Roboto-BoldCondensedItalic.ttf"
++Files=".../Engine/Content/Slate/Fonts/Roboto-BoldItalic.ttf"
++Files=".../Engine/Content/Slate/Fonts/Roboto-Italic.ttf"
++Files=".../Engine/Content/Slate/Fonts/Roboto-Light.ttf"
++Files=".../Engine/Content/SlateDebug/Fonts/LastResort.ttf"
++Files=".../Engine/Content/EngineMaterials/DefaultBloomKernel.*"
++Files=".../Engine/Content/EngineMaterials/BlueNoise.uexp"
+;+Files=".../Engine/Content/EngineMaterials/DefaultVirtualTextureMaterial.uexp"
+;+Files=".../Engine/Content/EngineMaterials/DefaultDeferredDecalMaterial.uexp"
+;+Files=".../Engine/Content/EngineMaterials/HairDefaultMaterial.uexp"
++Files=".../Engine/Content/EngineMaterials/WorldGridMaterial.uexp"
diff --git a/Templates/TP_HMI_Blank/Config/TemplateDefs.ini b/Templates/TP_HMI_Blank/Config/TemplateDefs.ini
new file mode 100644
--- /dev/null
+++ b/Templates/TP_HMI_Blank/Config/TemplateDefs.ini
@@ -0,0 +1,124 @@
+﻿; This file specifies how a template will be treated when creating a project from it.
+; It includes rules for ignoring files, replacing files and changing the contents of files.
+; This file is annotated to help with creating your own templates.
+
+[/Script/GameProjectGeneration.TemplateProjectDefs]
+
+; The categories that this template should appear in. The values in this list refer to the Key value given in TemplateCategories.ini.
+Categories=HMI
+
+; The localized display names for this template. Each entry specifies a Language code and Unicode text to display in the New Project dialog.
+LocalizedDisplayNames=(Language="en", Text="HMI Automotive Blank")
+
+; The localized descriptions for this template, which is displayed when the template is selected.
+LocalizedDescriptions=(Language="en", Text="A clean empty HMI automotive project with no code.")
+
+
+; The sort key to use for the template to bypass the default alphabetical sorting.
+;SortKey="_1"
+
+; A list of folders to avoid copying when creating a project from this template. 
+; These should contain files only of concern to the template itself and not desired in the user's project.
+FoldersToIgnore=Binaries
+FoldersToIgnore=Build
+FoldersToIgnore=Intermediate
+FoldersToIgnore=Saved
+FoldersToIgnore=Media
+
+; Is this a blank template? Blank templates do not specify a default map.
+bIsBlank=True
+PreviewFeatureLevel=1
+PreviewShaderFormatName=GLSL_ES3_1_ANDROID
+bPreviewFeatureLevelActive=True
+HardwareTarget=EHardwareClass::Mobile
+HardwareClass=EHardwareClass::Mobile
+GraphicsPreset=EGraphicsPreset::Scalable
+GraphicsPerformance=EGraphicsPreset::Scalable
+DefaultGraphicsPerformance=EGraphicsPreset::Scalable
+
+; The files to ignore copying, relative to the template root folder. 
+; %TEMPLATENAME% will be replaced by the template's folder name when resolving these. 
+; Filename comparisons are case-insensitive. 
+; For example, if this template is "TP_Blank" and the user creates "MyProject", then "TP_Blank.uproject" will be copied to "MyProject.uproject".
+FilesToIgnore="%TEMPLATENAME%.uproject"
+FilesToIgnore="%TEMPLATENAME%.png"
+FilesToIgnore="Config/TemplateDefs.ini"
+FilesToIgnore="Config/config.ini"
+FilesToIgnore="%TEMPLATENAME%.opensdf"
+FilesToIgnore="%TEMPLATENAME%.sdf"
+FilesToIgnore="%TEMPLATENAME%.v11.suo"
+FilesToIgnore="%TEMPLATENAME%.v12.suo"
+FilesToIgnore="%TEMPLATENAME%.sln"
+FilesToIgnore="Manifest.json"
+FilesToIgnore="contents.txt"
+
+; NOTE: For the purposes of the below examples, this template is considered to be named "TP_Blank", and the user has created a project called "MyProject".
+
+; Folders that will be renamed when creating a project from this template.
+; %PROJECTNAME% will be replaced by the project name specified by the user during creation when resolving these rules.
+; For example,  then "Source/TP_Blank" will be copied to "Source/MyProject".
+FolderRenames=(From="Source/%TEMPLATENAME%", To="Source/%PROJECTNAME%")
+FolderRenames=(From="Source/%TEMPLATENAME%Editor", To="Source/%PROJECTNAME%Editor")
+
+; Rules for renaming files when creating a project from this template. 
+; %TEMPLATENAME_UPPERCASE% and %PROJECTNAME_UPPERCASE% are the template and project names in upper case, eg. for "TP_Blank", its upper case representation is "TP_BLANK".
+; %TEMPLATENAME_LOWERCASE% and %PROJECTNAME_LOWERCASE% are the template and project names in lower case, eg. for "TP_Blank", its lower case representation is "tp_blank".
+; A rule does not have to include any of these modifiers. 
+; Extensions is a list of files to consider for these replacements. Eg. ("cpp", "h") will try to apply a rule only to files like "SomeFile.cpp" and "SomeFile.h".
+; bCaseSensitive specifies whether this replacement should be case-sensitive.
+
+; Here are some default rules for renaming files:
+
+; This rule states that a file called "TP_BLANK.cpp" will be renamed to "MYPROJECT.cpp".
+FilenameReplacements=(Extensions=("cpp","h","ini","cs"), From="%TEMPLATENAME_UPPERCASE%", To="%PROJECTNAME_UPPERCASE%", bCaseSensitive=true)
+
+; This rule states that a file called "tp_blank.ini" will be renamed to "myproject.ini"
+FilenameReplacements=(Extensions=("cpp","h","ini","cs"), From="%TEMPLATENAME_LOWERCASE%", To="%PROJECTNAME_LOWERCASE%", bCaseSensitive=true)
+
+; And this rule catches any irregularly-cased files, so that "Tp_bLANK.ini" will be renamed to "MyProject.ini".
+FilenameReplacements=(Extensions=("cpp","h","ini","cs"), From="%TEMPLATENAME%",To="%PROJECTNAME%",bCaseSensitive=false)
+
+; Rules for renaming strings in the contents of files. 
+; The same rules as above apply, but any usages of strings that match these rules inside the file will be changed. This includes all substrings!
+; For example, "TP_Blank" will be replaced by "MyProject", but "TP_Blank_Name" will also be replaced by "MyProject_Name".
+ReplacementsInFiles=(Extensions=("cpp","h","ini","cs"), From="%TEMPLATENAME_UPPERCASE%", To="%PROJECTNAME_UPPERCASE%", bCaseSensitive=true)
+ReplacementsInFiles=(Extensions=("cpp","h","ini","cs"), From="%TEMPLATENAME_LOWERCASE%", To="%PROJECTNAME_LOWERCASE%", bCaseSensitive=true)
+ReplacementsInFiles=(Extensions=("cpp","h","ini","cs"), From="%TEMPLATENAME%", To="%PROJECTNAME%", bCaseSensitive=false)
+
+; The following settings are not enabled for this project, but are included for completeness:
+
+; Whether or not this template can be used to create a project. Default: true
+; If this is set to false, then the template becomes a Feature Pack that can only be referenced from other templates, or added to a project from the Content Browser.
+;bAllowProjectCreation=false
+
+; Which settings on the project settings page should be hidden for this template. Default: none
+; Allowed values are:
+; All - Hides the settings page entirely
+; Languages - Choose between C++ and Blueprint versions.
+; HardwareTarget - Choose between Mobile and Desktop hardware targets.
+; GraphicsPreset - Choose between Maximum and Scalable graphics.
+; Raytracing - Choose whether raytracing should be enabled for the project.
+; StarterContent - Choose whether to add starter content to the project.
+; XR - Choose whether to enable Mixed Reality support for the project.
+;HiddenSettings=All
+;HiddenSettings=Languages
+HiddenSettings=HardwareTarget
+HiddenSettings=GraphicsPreset 
+HiddenSettings=Raytracing
+HiddenSettings=StarterContent 
+HiddenSettings=XR
+
+; The class types that this project uses.
+; This is only for display in the template description, and does not affect project creation.
+;ClassTypes=GameMode, Character, Actor
+
+; The asset types that this project uses.
+; This is only for display in the template description, and does not affect project creation.
+;AssetTypes=Skeleton, Static Mesh, Texture
+
+; Shared feature packs. The files in these packs listed in these structures marked as 'additionalfiles' will be copied on project generation.
+;SharedContentPacks=(MountName="2DSideScroller",DetailLevels=("Standard"))
+
+; Starter content feature pack to use if the user chooses to add starter content.
+; This defaults to StarterContent on desktop and MobileStarterContent on mobile.
+;StarterContent=TP_FirstPerson
diff --git a/Templates/TP_HMI_Blank/Source/TP_HMI_Blank.Target.cs b/Templates/TP_HMI_Blank/Source/TP_HMI_Blank.Target.cs
new file mode 100644
--- /dev/null
+++ b/Templates/TP_HMI_Blank/Source/TP_HMI_Blank.Target.cs
@@ -0,0 +1,15 @@
+// Copyright Epic Games, Inc. All Rights Reserved.
+
+using UnrealBuildTool;
+using System.Collections.Generic;
+
+public class TP_HMI_BlankTarget : TargetRules
+{
+	public TP_HMI_BlankTarget(TargetInfo Target) : base(Target)
+	{
+		Type = TargetType.Game;
+		DefaultBuildSettings = BuildSettingsVersion.V5;
+		IncludeOrderVersion = EngineIncludeOrderVersion.Unreal5_4;
+		ExtraModuleNames.Add("TP_HMI_Blank");
+	}
+}
diff --git a/Templates/TP_HMI_Blank/Source/TP_HMI_Blank/MyClass.cpp b/Templates/TP_HMI_Blank/Source/TP_HMI_Blank/MyClass.cpp
new file mode 100644
--- /dev/null
+++ b/Templates/TP_HMI_Blank/Source/TP_HMI_Blank/MyClass.cpp
@@ -0,0 +1,12 @@
+// Copyright Epic Games, Inc. All Rights Reserved.
+
+
+#include "MyClass.h"
+
+MyClass::MyClass()
+{
+}
+
+MyClass::~MyClass()
+{
+}
diff --git a/Templates/TP_HMI_Blank/Source/TP_HMI_Blank/MyClass.h b/Templates/TP_HMI_Blank/Source/TP_HMI_Blank/MyClass.h
new file mode 100644
--- /dev/null
+++ b/Templates/TP_HMI_Blank/Source/TP_HMI_Blank/MyClass.h
@@ -0,0 +1,15 @@
+// Copyright Epic Games, Inc. All Rights Reserved.
+
+#pragma once
+
+#include "CoreMinimal.h"
+
+/**
+ * 
+ */
+class TP_HMI_BLANK_API MyClass
+{
+public:
+	MyClass();
+	~MyClass();
+};
diff --git a/Templates/TP_HMI_Blank/Source/TP_HMI_Blank/TP_HMI_Blank.Build.cs b/Templates/TP_HMI_Blank/Source/TP_HMI_Blank/TP_HMI_Blank.Build.cs
new file mode 100644
--- /dev/null
+++ b/Templates/TP_HMI_Blank/Source/TP_HMI_Blank/TP_HMI_Blank.Build.cs
@@ -0,0 +1,28 @@
+// Copyright Epic Games, Inc. All Rights Reserved.
+
+using UnrealBuildTool;
+
+public class TP_HMI_Blank : ModuleRules
+{
+	public TP_HMI_Blank(ReadOnlyTargetRules Target) : base(Target)
+	{
+		PCHUsage = PCHUsageMode.UseExplicitOrSharedPCHs;
+	
+		PublicDependencyModuleNames.AddRange(new string[] { "Core", "CoreUObject", "Engine", "InputCore" });
+
+        PrivateDependencyModuleNames.AddRange(new string[] {  });
+
+        if (Target.Platform == UnrealTargetPlatform.Android)
+        {
+            PrivateDependencyModuleNames.AddRange(new string[] { "Launch" });
+        }
+
+        // Uncomment if you are using Slate UI
+        // PrivateDependencyModuleNames.AddRange(new string[] { "Slate", "SlateCore" });
+
+        // Uncomment if you are using online features
+        // PrivateDependencyModuleNames.Add("OnlineSubsystem");
+
+        // To include OnlineSubsystemSteam, add it to the plugins section in your uproject file with the Enabled attribute set to true
+    }
+}
diff --git a/Templates/TP_HMI_Blank/Source/TP_HMI_Blank/TP_HMI_Blank.cpp b/Templates/TP_HMI_Blank/Source/TP_HMI_Blank/TP_HMI_Blank.cpp
new file mode 100644
--- /dev/null
+++ b/Templates/TP_HMI_Blank/Source/TP_HMI_Blank/TP_HMI_Blank.cpp
@@ -0,0 +1,6 @@
+// Copyright Epic Games, Inc. All Rights Reserved.
+
+#include "TP_HMI_Blank.h"
+#include "Modules/ModuleManager.h"
+
+IMPLEMENT_PRIMARY_GAME_MODULE( FDefaultGameModuleImpl, TP_HMI_Blank, "TP_HMI_Blank" );
diff --git a/Templates/TP_HMI_Blank/Source/TP_HMI_Blank/TP_HMI_Blank.h b/Templates/TP_HMI_Blank/Source/TP_HMI_Blank/TP_HMI_Blank.h
new file mode 100644
--- /dev/null
+++ b/Templates/TP_HMI_Blank/Source/TP_HMI_Blank/TP_HMI_Blank.h
@@ -0,0 +1,6 @@
+// Copyright Epic Games, Inc. All Rights Reserved.
+
+#pragma once
+
+#include "CoreMinimal.h"
+
diff --git a/Templates/TP_HMI_Blank/Source/TP_HMI_BlankEditor.Target.cs b/Templates/TP_HMI_Blank/Source/TP_HMI_BlankEditor.Target.cs
new file mode 100644
--- /dev/null
+++ b/Templates/TP_HMI_Blank/Source/TP_HMI_BlankEditor.Target.cs
@@ -0,0 +1,15 @@
+// Copyright Epic Games, Inc. All Rights Reserved.
+
+using UnrealBuildTool;
+using System.Collections.Generic;
+
+public class TP_HMI_BlankEditorTarget : TargetRules
+{
+	public TP_HMI_BlankEditorTarget(TargetInfo Target) : base(Target)
+	{
+		Type = TargetType.Editor;
+		DefaultBuildSettings = BuildSettingsVersion.V5;
+		IncludeOrderVersion = EngineIncludeOrderVersion.Unreal5_4;
+		ExtraModuleNames.Add("TP_HMI_Blank");
+	}
+}
diff --git a/Templates/TP_HMI_Blank/TP_HMI_Blank.uproject b/Templates/TP_HMI_Blank/TP_HMI_Blank.uproject
new file mode 100644
--- /dev/null
+++ b/Templates/TP_HMI_Blank/TP_HMI_Blank.uproject
@@ -0,0 +1,225 @@
+{
+	"FileVersion": 3,
+	"EngineAssociation": "",
+	"Category": "",
+	"Description": "",
+	"Modules": [
+		{
+			"Name": "TP_HMI_Blank",
+			"Type": "Runtime",
+			"LoadingPhase": "Default"
+		}
+	],
+	"Plugins": [
+		{
+			"Name": "SteamVR",
+			"Enabled": false,
+			"SupportedTargetPlatforms": [
+				"Win32",
+				"Win64",
+				"Linux",
+				"Mac"
+			]
+		},
+		{
+			"Name": "MagicLeapMedia",
+			"Enabled": false,
+			"SupportedTargetPlatforms": [
+				"Lumin"
+			]
+		},
+		{
+			"Name": "MagicLeap",
+			"Enabled": false,
+			"SupportedTargetPlatforms": [
+				"Lumin",
+				"Mac",
+				"Win64"
+			]
+		},
+		{
+			"Name": "OculusVR",
+			"Enabled": false,
+			"SupportedTargetPlatforms": [
+				"Win32",
+				"Win64",
+				"Android"
+			]
+		},
+		{
+			"Name": "AndroidSingleInstanceService",
+			"Enabled": true
+		},
+		{
+			"Name": "Niagra",
+			"Enabled": false
+		},
+		{
+			"Name": "Paper2D",
+			"Enabled": false
+		},
+		{
+			"Name": "AISupport",
+			"Enabled": false
+		},
+		{
+			"Name": "LightPropagationVolume",
+			"Enabled": false
+		},
+		{
+			"Name": "AnimationSharing",
+			"Enabled": false
+		},
+		{
+			"Name": "UObjectPlugin",
+			"Enabled": false
+		},
+		{
+			"Name": "MaterialAnalyzer",
+			"Enabled": false
+		},
+		{
+			"Name": "AssetManagerEditor",
+			"Enabled": false
+		},
+		{
+			"Name": "FacialAnimation",
+			"Enabled": false
+		},
+		{
+			"Name": "AlembicImporter",
+			"Enabled": false
+		},
+		{
+			"Name": "GeometryCache",
+			"Enabled": false
+		},
+		{
+			"Name": "AutomationUtils",
+			"Enabled": false
+		},
+		{
+			"Name": "BackChannel",
+			"Enabled": false
+		},
+		{
+			"Name": "CharacterAI",
+			"Enabled": false
+		},
+		{
+			"Name": "MediaCompositing",
+			"Enabled": false
+		},
+		{
+			"Name": "ImgMedia",
+			"Enabled": false
+		},
+		{
+			"Name": "ActorSequence",
+			"Enabled": false
+		},
+		{
+			"Name": "NUTUnrealEngine4",
+			"Enabled": false
+		},
+		{
+			"Name": "NUTUnrealEngine4_Private",
+			"Enabled": false
+		},
+		{
+			"Name": "NetcodeUnitTest",
+			"Enabled": false
+		},
+		{
+			"Name": "OnlineSubsystemGooglePlay",
+			"Enabled": false,
+			"SupportedTargetPlatforms": [
+				"Android"
+			]
+		},
+		{
+			"Name": "OnlineSubsystemIOS",
+			"Enabled": false,
+			"SupportedTargetPlatforms": [
+				"IOS",
+				"TVOS"
+			]
+		},
+		{
+			"Name": "OnlineSubsystemNull",
+			"Enabled": false
+		},
+		{
+			"Name": "OnlineSubsystemUtils",
+			"Enabled": false
+		},
+		{
+			"Name": "OnlineSubsystem",
+			"Enabled": false
+		},
+		{
+			"Name": "AppleImageUtils",
+			"Enabled": false
+		},
+		{
+			"Name": "ArchVisCharacter",
+			"Enabled": false
+		},
+		{
+			"Name": "AudioCapture",
+			"Enabled": false
+		},
+		{
+			"Name": "CableComponent",
+			"Enabled": false
+		},
+		{
+			"Name": "CustomMeshComponent",
+			"Enabled": false
+		},
+		{
+			"Name": "EditableMesh",
+			"Enabled": false
+		},
+		{
+			"Name": "ExampleDeviceProfileSelector",
+			"Enabled": false
+		},
+		{
+			"Name": "GoogleCloudMessaging",
+			"Enabled": false
+		},
+		{
+			"Name": "LocationServicesBPLibrary",
+			"Enabled": false
+		},
+		{
+			"Name": "PhysXVehicles",
+			"Enabled": false
+		},
+		{
+			"Name": "ProceduralMeshComponent",
+			"Enabled": false
+		},
+		{
+			"Name": "RuntimePhysXCooking",
+			"Enabled": false
+		},
+		{
+			"Name": "MLSDK",
+			"Enabled": false
+		},
+		{
+			"Name": "DatasmithContent",
+			"Enabled": false
+		},
+		{
+			"Name": "VariantManager",
+			"Enabled": true
+		},
+		{
+			"Name": "GooglePAD",
+			"Enabled": false
+		}
+	]
+}
diff --git a/Templates/TP_HMI_Blank/contents.txt b/Templates/TP_HMI_Blank/contents.txt
new file mode 100644
--- /dev/null
+++ b/Templates/TP_HMI_Blank/contents.txt
@@ -0,0 +1,5 @@
+../../../Templates/TP_HMI_Blank/manifest.json
+../../../Templates/TP_HMI_Blank/Config/config.ini
+../../../Templates/TP_HMI_Blank/Media/TP_HMI_Blank.png
+../../../Templates/TP_HMI_Blank/Media/TP_HMI_Blank_Preview.png
+../../../Templates/TP_HMI_Blank/Media/TP_HMI_Blank_TemplateIcon.png
diff --git a/Templates/TP_HMI_Blank/manifest.json b/Templates/TP_HMI_Blank/manifest.json
new file mode 100644
--- /dev/null
+++ b/Templates/TP_HMI_Blank/manifest.json
@@ -0,0 +1,29 @@
+﻿{
+	"Version": 1,
+	"Name":
+	[
+		{
+			"Language": "en",
+			"Text": "Blank"
+		}
+	],
+	"Description":
+	[
+		{
+			"Language": "en",
+			"Text": "An empty project with some basic HMI autotmotive framework code classes created."
+		}
+	],
+	"AssetTypes": [],	
+	"SearchTags": [],
+	"ClassTypes": "",
+	"Category": "",
+	"FocusAsset": "",
+	"Thumbnail": "TP_HMI_Blank.png",
+	"Screenshots":
+	[
+		"TP_HMI_Blank_Preview.png"
+	],
+	"SortKey": "_1",
+	"AdditionalFeaturePacks": []
+}
diff --git a/Templates/TemplateCategories.ini b/Templates/TemplateCategories.ini
--- a/Templates/TemplateCategories.ini
+++ b/Templates/TemplateCategories.ini
@@ -21,3 +21,5 @@ Categories=(Key="AEC", LocalizedDisplayNames=((Language="en",Text="Architecture"
 Categories=(Key="MFG", LocalizedDisplayNames=((Language="en",Text="Automotive\nProduct Design &\nManufacturing"),(Language="zh-Hans",Text="汽车、产品设计和制造"),(Language="ja",Text="自動車、プロダクト デザイン、製造"),(Language="ko",Text="자동차, 제품 디자인 및 제조")), LocalizedDescriptions=((Language="en",Text="Find templates for multi-user design reviews, photobooth studio environments, and product configurators."),(Language="zh-Hans",Text="为多用户设计评审、Photobooth工作室环境和产品配置器寻找适合的模板。"),(Language="ja",Text="複数のユーザーによるレビューや、フォトブース スタジオの環境背景、製品コンフィギュレーターのためのテンプレートを見つけることができます。"),(Language="ko",Text="멀티 유저 디자인 리뷰, 포토부스 스튜디오 환경 및 제품 컨피규레이터를 위한 템플릿을 살펴보세요.")), Icon="Media/AutomotiveProduct_2x.png", IsMajorCategory=true)
 
 Categories=(Key="SIM", LocalizedDisplayNames=((Language="en",Text="Simulation"),(Language="zh-Hans",Text="模拟"),(Language="ja",Text="シミュレーション"),(Language="ko",Text="시뮬레이션")), LocalizedDescriptions=((Language="en",Text="Find templates for Outdoor georeferenced environments, First person, Driving, Multichannels and AR/VR projects."),(Language="zh-Hans",Text="查看可用于室外地理参考场景、第一人称项目、驾驶类场景、多通道项目，以及AR/VR项目的模板。"),(Language="ja",Text="屋外ジオリファレンス環境、ファーストパーソン、運転、マルチチャンネル、AR/VR プロジェクト向けのテンプレートです。"),(Language="ko",Text="아웃도어 지오레퍼런스드 환경, 1인칭, 드라이빙, 멀티 채널 및 AR/VR 프로젝트를 위한 템플릿을 찾아보세요.")), Icon="Media/Simulation_2x.png", IsMajorCategory=true)
+
+Categories=(Key="HMI", LocalizedDisplayNames=((Language="en",Text="Automotive, Virtual Cockpit")), LocalizedDescriptions=((Language="en",Text="Find templates for Virtual Cockpits and vehicle HMI.")), Icon="TP_HMI_Blank/Media/TP_HMI_Blank_TemplateIcon.png", IsMajorCategory=true)
\ No newline at end of file
